Fixed a bug in the Apsalar plugin that reported the name of an attribute as the value

diff --git a/Engine/Plugins/Runtime/Analytics/Apsalar/IOSApsalar/Source/IOSApsalar/Private/IOSApsalar.cpp b/Engine/Plugins/Runtime/Analytics/Apsalar/IOSApsalar/Source/IOSApsalar/Private/IOSApsalar.cpp
--- a/Engine/Plugins/Runtime/Analytics/Apsalar/IOSApsalar/Source/IOSApsalar/Private/IOSApsalar.cpp
+++ b/Engine/Plugins/Runtime/Analytics/Apsalar/IOSApsalar/Source/IOSApsalar/Private/IOSApsalar.cpp
@@ -219,7 +219,7 @@ void FAnalyticsProviderApsalar::RecordEvent(const FString& EventName, const TArr
 			for	(auto Attr : Attributes)
 			{
 				NSString* AttrName = [NSString stringWithFString : Attr.AttrName];
-				NSString* AttrValue = [NSString stringWithFString : Attr.AttrName];
+				NSString* AttrValue = [NSString stringWithFString : Attr.AttrValue];
 				[AttributesDict setValue:AttrValue forKey:AttrName];
 			}
 			[Apsalar event:ConvertedEventName withArgs:AttributesDict];
